- Renamed Desktop target platform to AllDesktop to reduce conflicts with Desktop directories [UEPLAT-894]

diff --git a/Engine/Build/JunkManifest.txt b/Engine/Build/JunkManifest.txt
--- a/Engine/Build/JunkManifest.txt
+++ b/Engine/Build/JunkManifest.txt
@@ -72,3 +72,6 @@ Engine\Plugins\Runtime\TwitchLiveStreaming\Binaries\Win64\*-TwitchLiveStreaming*
 Engine\Binaries\Mac\UE4Editor.app\Contents\MacOS\UE4Editor-Android_ES31TargetPlatform.dylib
 Engine\Binaries\Mac\UnrealFrontend.app\Contents\MacOS\UnrealFrontend-Android_ES31TargetPlatform.dylib
 Engine\Binaries\Win64\Android\*Android_ES31*.*
+Engine\Binaries\Mac\UE4Editor.app\Contents\MacOS\UE4Editor-DesktopTargetPlatform.dylib
+Engine\Binaries\Mac\UnrealFrontend.app\Contents\MacOS\UnrealFrontend-DesktopTargetPlatform.dylib
+Engine\Binaries\Win64\*-DesktopTargetPlatform*.*
diff --git a/Engine/Source/Developer/AllDesktopTargetPlatform/AllDesktopTargetPlatform.Build.cs b/Engine/Source/Developer/AllDesktopTargetPlatform/AllDesktopTargetPlatform.Build.cs
new file mode 100644
--- /dev/null
+++ b/Engine/Source/Developer/AllDesktopTargetPlatform/AllDesktopTargetPlatform.Build.cs
@@ -0,0 +1,33 @@
+// Copyright 1998-2015 Epic Games, Inc. All Rights Reserved.
+
+using UnrealBuildTool;
+
+public class AllDesktopTargetPlatform : ModuleRules
+{
+	public AllDesktopTargetPlatform(TargetInfo Target)
+	{
+		PrivateDependencyModuleNames.AddRange(
+			new string[]
+			{
+				"Core",
+				"CoreUObject",
+				"TargetPlatform",
+				"DesktopPlatform",
+				"LaunchDaemonMessages",
+				"Projects"
+			}
+		);
+
+		PrivateIncludePathModuleNames.AddRange(
+			new string[] {
+				"Messaging",
+				"TargetDeviceServices",
+			}
+		);
+
+		if (UEBuildConfiguration.bCompileAgainstEngine)
+		{
+			PrivateDependencyModuleNames.Add("Engine");
+		}
+	}
+}
diff --git a/Engine/Source/Developer/AllDesktopTargetPlatform/Private/AllDesktopTargetPlatform.cpp b/Engine/Source/Developer/AllDesktopTargetPlatform/Private/AllDesktopTargetPlatform.cpp
new file mode 100644
--- /dev/null
+++ b/Engine/Source/Developer/AllDesktopTargetPlatform/Private/AllDesktopTargetPlatform.cpp
@@ -0,0 +1,82 @@
+// Copyright 1998-2015 Epic Games, Inc. All Rights Reserved.
+
+/*=============================================================================
+	AllDesktopTargetPlatform.cpp: Implements the FDesktopTargetPlatform class.
+=============================================================================*/
+
+#include "AllDesktopTargetPlatformPrivatePCH.h"
+#include "IProjectManager.h"
+
+/* FAllDesktopTargetPlatform structors
+ *****************************************************************************/
+
+FAllDesktopTargetPlatform::FAllDesktopTargetPlatform()
+{
+#if WITH_ENGINE
+	// use non-platform specific settings
+	FConfigCacheIni::LoadLocalIniFile(EngineSettings, TEXT("Engine"), true, NULL);
+	StaticMeshLODSettings.Initialize(EngineSettings);
+#endif // #if WITH_ENGINE
+
+}
+
+
+FAllDesktopTargetPlatform::~FAllDesktopTargetPlatform()
+{
+}
+
+
+
+/* ITargetPlatform interface
+ *****************************************************************************/
+
+#if WITH_ENGINE
+
+
+void FAllDesktopTargetPlatform::GetAllPossibleShaderFormats( TArray<FName>& OutFormats ) const
+{
+	static FName NAME_PCD3D_SM5(TEXT("PCD3D_SM5"));
+	static FName NAME_PCD3D_SM4(TEXT("PCD3D_SM4"));
+	static FName NAME_GLSL_150(TEXT("GLSL_150"));
+	static FName NAME_GLSL_430(TEXT("GLSL_430"));
+	static FName NAME_GLSL_150_MAC(TEXT("GLSL_150_MAC"));
+	
+#if PLATFORM_WINDOWS
+	// right now, only windows can properly compile D3D shaders (this won't corrupt the DDC, but it will 
+	// make it so that packages cooked on Mac/Linux will only run on Windows with -opengl)
+	OutFormats.AddUnique(NAME_PCD3D_SM5);
+	OutFormats.AddUnique(NAME_PCD3D_SM4);
+#endif
+	OutFormats.AddUnique(NAME_GLSL_150);
+	OutFormats.AddUnique(NAME_GLSL_430);
+	OutFormats.AddUnique(NAME_GLSL_150_MAC);
+}
+
+
+void FAllDesktopTargetPlatform::GetAllTargetedShaderFormats( TArray<FName>& OutFormats ) const
+{
+	GetAllPossibleShaderFormats(OutFormats);
+}
+
+void FAllDesktopTargetPlatform::GetTextureFormats( const UTexture* Texture, TArray<FName>& OutFormats ) const
+{
+	// just use the standard texture format name for this texture (without DX11 texture support)
+	OutFormats.Add(GetDefaultTextureFormatName(Texture, EngineSettings, false));
+}
+
+
+FName FAllDesktopTargetPlatform::GetWaveFormat( const class USoundWave* Wave ) const
+{
+	static FName NAME_OGG(TEXT("OGG"));
+	static FName NAME_OPUS(TEXT("OPUS"));
+	
+	if (Wave->IsStreaming())
+	{
+		// @todo desktop platform: Does Mac support OPUS?
+		return NAME_OPUS;
+	}
+	
+	return NAME_OGG;
+}
+
+#endif // WITH_ENGINE
diff --git a/Engine/Source/Developer/AllDesktopTargetPlatform/Private/AllDesktopTargetPlatform.h b/Engine/Source/Developer/AllDesktopTargetPlatform/Private/AllDesktopTargetPlatform.h
new file mode 100644
--- /dev/null
+++ b/Engine/Source/Developer/AllDesktopTargetPlatform/Private/AllDesktopTargetPlatform.h
@@ -0,0 +1,174 @@
+// Copyright 1998-2015 Epic Games, Inc. All Rights Reserved.
+
+/*=============================================================================
+	AllDesktopTargetPlatform.h: Declares the FDesktopTargetPlatform class.
+=============================================================================*/
+
+#pragma once
+
+#if WITH_ENGINE
+#include "StaticMeshResources.h"
+#endif // WITH_ENGINE
+
+class FAllDesktopPlatformProperties : public FGenericPlatformProperties
+{
+public:
+	static FORCEINLINE const char* PlatformName()
+	{
+		return "AllDesktop";
+	}
+	static FORCEINLINE const char* IniPlatformName()
+	{
+		// this uses generic, non-platform-specific .ini files
+		return "";
+	}
+	static FORCEINLINE bool HasEditorOnlyData()
+	{
+		return false;
+	}
+	static FORCEINLINE bool RequiresCookedData()
+	{
+		return true;
+	}
+};
+
+/**
+ * FDesktopTargetPlatform, abstraction for cooking iOS platforms
+ */
+class FAllDesktopTargetPlatform
+	: public TTargetPlatformBase<FAllDesktopPlatformProperties>
+{
+public:
+
+	/**
+	 * Default constructor.
+	 */
+	FAllDesktopTargetPlatform();
+
+	/**
+	 * Destructor.
+	 */
+	~FAllDesktopTargetPlatform();
+
+public:
+
+	// Begin TTargetPlatformBase interface
+
+	virtual bool IsServerOnly( ) const override
+	{
+		return false;
+	}
+
+	// End TTargetPlatformBase interface
+
+public:
+
+	// Begin ITargetPlatform interface
+
+	virtual void EnableDeviceCheck(bool OnOff) override
+	{
+	}
+
+	virtual void GetAllDevices( TArray<ITargetDevicePtr>& OutDevices ) const override
+	{
+	}
+
+	virtual ECompressionFlags GetBaseCompressionMethod() const override
+	{
+		return COMPRESS_ZLIB;
+	}
+
+	virtual bool GenerateStreamingInstallManifest(const TMultiMap<FString, int32>& ChunkMap, const TSet<int32>& ChunkIDsInUse) const override
+	{
+		return true;
+	}
+
+	virtual ITargetDevicePtr GetDefaultDevice( ) const override
+	{
+		return NULL;
+	}
+
+	virtual ITargetDevicePtr GetDevice( const FTargetDeviceId& DeviceId ) override
+	{
+		return NULL;
+	}
+
+	virtual bool IsRunningPlatform( ) const override
+	{
+		return false;
+	}
+
+	virtual bool SupportsFeature( ETargetPlatformFeatures Feature ) const override
+	{
+		return TTargetPlatformBase<FAllDesktopPlatformProperties>::SupportsFeature(Feature);
+	}
+
+	virtual bool IsSdkInstalled(bool bProjectHasCode, FString& OutTutorialPath) const override
+	{
+		return true;
+	}
+
+	virtual int32 CheckRequirements(const FString& ProjectPath, bool bProjectHasCode, FString& OutTutorialPath) const override
+	{
+		return ETargetPlatformReadyStatus::Ready;
+	}
+
+
+#if WITH_ENGINE
+
+	virtual void GetAllPossibleShaderFormats( TArray<FName>& OutFormats ) const override;
+
+	virtual void GetAllTargetedShaderFormats( TArray<FName>& OutFormats ) const override;
+
+	virtual const class FStaticMeshLODSettings& GetStaticMeshLODSettings( ) const override
+	{
+		return StaticMeshLODSettings;
+	}
+
+	virtual void GetTextureFormats(const UTexture* Texture, TArray<FName>& OutFormats) const override;
+	
+	virtual const UTextureLODSettings& GetTextureLODSettings() const override
+	{
+		return *TextureLODSettings;
+	}
+
+	virtual void RegisterTextureLODSettings(const UTextureLODSettings* InTextureLODSettings) override
+	{
+		TextureLODSettings = InTextureLODSettings;
+	}
+
+	virtual FName GetWaveFormat( const class USoundWave* Wave ) const override;
+#endif // WITH_ENGINE
+
+
+	DECLARE_DERIVED_EVENT(FAllDesktopTargetPlatform, ITargetPlatform::FOnTargetDeviceDiscovered, FOnTargetDeviceDiscovered);
+	virtual FOnTargetDeviceDiscovered& OnDeviceDiscovered( ) override
+	{
+		return DeviceDiscoveredEvent;
+	}
+
+	DECLARE_DERIVED_EVENT(FAllDesktopTargetPlatform, ITargetPlatform::FOnTargetDeviceLost, FOnTargetDeviceLost);
+	virtual FOnTargetDeviceLost& OnDeviceLost( ) override
+	{
+		return DeviceLostEvent;
+	}
+
+private:
+
+#if WITH_ENGINE
+	// Holds the Engine INI settings, for quick use.
+	FConfigFile EngineSettings;
+
+	// Holds the cache of the target LOD settings.
+	const UTextureLODSettings* TextureLODSettings;
+
+	// Holds the static mesh LOD settings.
+	FStaticMeshLODSettings StaticMeshLODSettings;
+#endif // WITH_ENGINE
+
+	// Holds an event delegate that is executed when a new target device has been discovered.
+	FOnTargetDeviceDiscovered DeviceDiscoveredEvent;
+
+	// Holds an event delegate that is executed when a target device has been lost, i.e. disconnected or timed out.
+	FOnTargetDeviceLost DeviceLostEvent;
+};
diff --git a/Engine/Source/Developer/AllDesktopTargetPlatform/Private/AllDesktopTargetPlatformModule.cpp b/Engine/Source/Developer/AllDesktopTargetPlatform/Private/AllDesktopTargetPlatformModule.cpp
new file mode 100644
--- /dev/null
+++ b/Engine/Source/Developer/AllDesktopTargetPlatform/Private/AllDesktopTargetPlatformModule.cpp
@@ -0,0 +1,57 @@
+// Copyright 1998-2015 Epic Games, Inc. All Rights Reserved.
+
+#include "AllDesktopTargetPlatformPrivatePCH.h"
+
+/**
+ * Holds the target platform singleton.
+ */
+static ITargetPlatform* Singleton = NULL;
+
+
+/**
+ * Module for a generic target platform for desktop platforms
+ */
+class FAllDesktopTargetPlatformModule : public ITargetPlatformModule
+{
+public:
+
+	/**
+	 * Destructor.
+	 */
+	FAllDesktopTargetPlatformModule()
+	{
+		Singleton = NULL;
+	}
+
+
+public:
+
+	// Begin ITargetPlatformModule interface
+
+	virtual ITargetPlatform* GetTargetPlatform() override
+	{
+		if (Singleton == NULL)
+		{
+			Singleton = new FAllDesktopTargetPlatform();
+		}
+
+		return Singleton;
+	}
+
+	// End ITargetPlatformModule interface
+
+
+public:
+
+	// Begin IModuleInterface interface
+	virtual void StartupModule() override
+	{
+	}
+
+	virtual void ShutdownModule() override
+	{
+	}
+	// End IModuleInterface interface
+};
+
+IMPLEMENT_MODULE( FAllDesktopTargetPlatformModule, AllDesktopTargetPlatform);
diff --git a/Engine/Source/Developer/AllDesktopTargetPlatform/Private/AllDesktopTargetPlatformPrivatePCH.h b/Engine/Source/Developer/AllDesktopTargetPlatform/Private/AllDesktopTargetPlatformPrivatePCH.h
new file mode 100644
--- /dev/null
+++ b/Engine/Source/Developer/AllDesktopTargetPlatform/Private/AllDesktopTargetPlatformPrivatePCH.h
@@ -0,0 +1,33 @@
+// Copyright 1998-2015 Epic Games, Inc. All Rights Reserved.
+
+/*=============================================================================
+	IOSTargetPlatformPrivatePCH.h: Pre-compiled header file for the IOSTargetPlatform module.
+=============================================================================*/
+
+#pragma once
+
+
+/* Dependencies
+ *****************************************************************************/
+
+#include "Core.h"
+#include "ModuleInterface.h"
+#include "ModuleManager.h"
+#include "LaunchDaemonMessages.h"
+#include "Messaging.h"
+#include "TargetDeviceServices.h"
+#include "GenericPlatformProperties.h"
+
+#if WITH_ENGINE
+	#include "Engine.h"
+#endif
+
+#include "TargetPlatform.h"
+#include "TargetPlatformBase.h"
+
+
+/* Private includes
+ *****************************************************************************/
+
+#include "AllDesktopTargetPlatform.h"
+ 
\ No newline at end of file
diff --git a/Engine/Source/Developer/DesktopPlatform/Private/PlatformInfo.cpp b/Engine/Source/Developer/DesktopPlatform/Private/PlatformInfo.cpp
--- a/Engine/Source/Developer/DesktopPlatform/Private/PlatformInfo.cpp
+++ b/Engine/Source/Developer/DesktopPlatform/Private/PlatformInfo.cpp
@@ -110,7 +110,7 @@ static const FPlatformInfo PlatformInfoArray[] = {
 
 	BuildPlatformInfo(TEXT("XboxOne"),					TEXT("XboxOne"),			LOCTEXT("XboxOne", "Xbox One"),								EPlatformType::Game,		EPlatformFlags::None,			FPlatformIconPaths(TEXT("Launcher/XboxOneTarget/Platform_XboxOne_24x"), TEXT("Launcher/XboxOneTarget/Platform_XboxOne_128x")),						TEXT(""),												TEXT("XboxOne"),	EPlatformSDKStatus::Unknown,	TEXT("/Platforms/XboxOne/GettingStarted"),																IsAvailableOnWindows,						false,				false),
 
-	BuildPlatformInfo(TEXT("Desktop"),					TEXT("Desktop"),			LOCTEXT("DesktopTargetPlatDisplay", "Desktop (Win+Mac+Linux)"),	EPlatformType::Game,	EPlatformFlags::None,			FPlatformIconPaths(TEXT("Launcher/DesktopTarget/Platform_Desktop_24x"), TEXT("Launcher/DesktopTarget/Platform_Desktop_128x")),						TEXT(""),												TEXT(""),			EPlatformSDKStatus::Unknown,	TEXT(""),																							false /*IsAvailableOnWindows || IsAvailableOnMac || IsAvailableOnLinux*/,	true, false),
+	BuildPlatformInfo(TEXT("AllDesktop"),					TEXT("AllDesktop"),			LOCTEXT("DesktopTargetPlatDisplay", "Desktop (Win+Mac+Linux)"),	EPlatformType::Game,	EPlatformFlags::None,			FPlatformIconPaths(TEXT("Launcher/DesktopTarget/Platform_Desktop_24x"), TEXT("Launcher/DesktopTarget/Platform_Desktop_128x")),						TEXT(""),												TEXT(""),			EPlatformSDKStatus::Unknown,	TEXT(""),																							false /*IsAvailableOnWindows || IsAvailableOnMac || IsAvailableOnLinux*/,	true, false),
 
 	//BuildPlatformInfo(TEXT("WinRT"),					TEXT("WinRT"),				LOCTEXT("WinRT", "Windows RT"),								EPlatformType::Game,		EPlatformFlags::None,			FPlatformIconPaths(TEXT("Launcher/WindowsTarget/Platform_WindowsNoEditor_24x"), TEXT("Launcher/WindowsTarget/Platform_WindowsNoEditor_128x")),		TEXT("")),
 };
diff --git a/Engine/Source/Developer/DesktopTargetPlatform/DesktopTargetPlatform.Build.cs b/Engine/Source/Developer/DesktopTargetPlatform/DesktopTargetPlatform.Build.cs
deleted file mode 100644
diff --git a/Engine/Source/Developer/DesktopTargetPlatform/Private/DesktopTargetPlatform.cpp b/Engine/Source/Developer/DesktopTargetPlatform/Private/DesktopTargetPlatform.cpp
deleted file mode 100644
diff --git a/Engine/Source/Developer/DesktopTargetPlatform/Private/DesktopTargetPlatform.h b/Engine/Source/Developer/DesktopTargetPlatform/Private/DesktopTargetPlatform.h
deleted file mode 100644
diff --git a/Engine/Source/Developer/DesktopTargetPlatform/Private/DesktopTargetPlatformModule.cpp b/Engine/Source/Developer/DesktopTargetPlatform/Private/DesktopTargetPlatformModule.cpp
deleted file mode 100644
diff --git a/Engine/Source/Developer/DesktopTargetPlatform/Private/DesktopTargetPlatformPrivatePCH.h b/Engine/Source/Developer/DesktopTargetPlatform/Private/DesktopTargetPlatformPrivatePCH.h
deleted file mode 100644
diff --git a/Engine/Source/Programs/AutomationTool/AllDesktop/AllDesktop.Automation.csproj b/Engine/Source/Programs/AutomationTool/AllDesktop/AllDesktop.Automation.csproj
new file mode 100644
--- /dev/null
+++ b/Engine/Source/Programs/AutomationTool/AllDesktop/AllDesktop.Automation.csproj
@@ -0,0 +1,75 @@
+<?xml version="1.0" encoding="utf-8"?>
+<Project ToolsVersion="4.0" DefaultTargets="Build" xmlns="http://schemas.microsoft.com/developer/msbuild/2003">
+  <Import Project="$(MSBuildExtensionsPath)\$(MSBuildToolsVersion)\Microsoft.Common.props" Condition="Exists('$(MSBuildExtensionsPath)\$(MSBuildToolsVersion)\Microsoft.Common.props')" />
+  <PropertyGroup>
+    <Configuration Condition=" '$(Configuration)' == '' ">Debug</Configuration>
+    <Platform Condition=" '$(Platform)' == '' ">AnyCPU</Platform>
+    <ProjectGuid>{A9CD1A06-764A-4C53-86CE-2F34B94F0283}</ProjectGuid>
+    <OutputType>Library</OutputType>
+    <AppDesignerFolder>Properties</AppDesignerFolder>
+    <RootNamespace>AllDesktop.Automation</RootNamespace>
+    <AssemblyName>AllDesktop.Automation</AssemblyName>
+    <TargetFrameworkVersion>v4.0</TargetFrameworkVersion>
+    <FileAlignment>512</FileAlignment>
+  </PropertyGroup>
+  <PropertyGroup Condition=" '$(Configuration)|$(Platform)' == 'Debug|AnyCPU' ">
+    <DebugSymbols>true</DebugSymbols>
+    <DebugType>full</DebugType>
+    <Optimize>false</Optimize>
+    <OutputPath>..\..\..\..\Binaries\DotNET\AutomationScripts\AllDesktop\</OutputPath>
+    <DefineConstants>DEBUG;TRACE</DefineConstants>
+    <ErrorReport>prompt</ErrorReport>
+    <WarningLevel>4</WarningLevel>
+    <PlatformTarget>AnyCPU</PlatformTarget>
+    <TreatWarningsAsErrors>true</TreatWarningsAsErrors>
+  </PropertyGroup>
+  <PropertyGroup Condition=" '$(Configuration)|$(Platform)' == 'Development|AnyCPU' ">
+    <DebugType>pdbonly</DebugType>
+    <Optimize>true</Optimize>
+    <OutputPath>..\..\..\..\Binaries\DotNET\AutomationScripts\AllDesktop\</OutputPath>
+    <DefineConstants>TRACE</DefineConstants>
+    <ErrorReport>prompt</ErrorReport>
+    <WarningLevel>4</WarningLevel>
+    <PlatformTarget>AnyCPU</PlatformTarget>
+    <TreatWarningsAsErrors>true</TreatWarningsAsErrors>
+  </PropertyGroup>
+  <ItemGroup>
+    <Reference Include="Ionic.Zip.Reduced">
+      <HintPath>..\..\..\..\Binaries\DotNET\Ionic.Zip.Reduced.dll</HintPath>
+    </Reference>
+    <Reference Include="System" />
+    <Reference Include="System.Core" />
+    <Reference Include="System.Runtime.Serialization" />
+    <Reference Include="System.Web" />
+    <Reference Include="System.Web.Extensions" />
+    <Reference Include="System.Xml.Linq" />
+    <Reference Include="System.Data.DataSetExtensions" />
+    <Reference Include="Microsoft.CSharp" />
+    <Reference Include="System.Data" />
+    <Reference Include="System.Xml" />
+    <Reference Include="OneSky">
+      <HintPath>..\..\..\..\Binaries\DotNET\OneSky.dll</HintPath>
+      <Private>False</Private>
+    </Reference>
+  </ItemGroup>
+  <ItemGroup>
+    <Compile Include="..\..\DotNETCommon\MetaData.cs">
+      <Link>Properties\MetaData.cs</Link>
+    </Compile>
+    <Compile Include="AllDesktopPlatform.Automation.cs" />
+    <Compile Include="Properties\AssemblyInfo.cs" />
+  </ItemGroup>
+  <ItemGroup>
+    <ProjectReference Include="..\AutomationTool.csproj">
+      <Project>{A35B0261-D868-475E-A098-EB14CE6DC89F}</Project>
+      <Name>AutomationTool</Name>
+      <Private>False</Private>
+    </ProjectReference>
+    <ProjectReference Include="..\..\UnrealBuildTool\UnrealBuildTool.csproj">
+      <Project>{FD7C5E1A-CFE4-4FD5-A525-1EB1599A39AC}</Project>
+      <Name>UnrealBuildTool</Name>
+      <Private>False</Private>
+    </ProjectReference>
+  </ItemGroup>
+  <Import Project="$(MSBuildToolsPath)\Microsoft.CSharp.targets" />
+</Project>
\ No newline at end of file
diff --git a/Engine/Source/Programs/AutomationTool/AllDesktop/AllDesktopPlatform.Automation.cs b/Engine/Source/Programs/AutomationTool/AllDesktop/AllDesktopPlatform.Automation.cs
new file mode 100644
--- /dev/null
+++ b/Engine/Source/Programs/AutomationTool/AllDesktop/AllDesktopPlatform.Automation.cs
@@ -0,0 +1,144 @@
+﻿// Copyright 1998-2015 Epic Games, Inc. All Rights Reserved.
+using System;
+using System.Collections.Generic;
+using System.Linq;
+using System.Text;
+using System.IO;
+using System.Diagnostics;
+using System.Net.NetworkInformation;
+using System.Threading;
+using AutomationTool;
+using UnrealBuildTool;
+using Ionic.Zip;
+
+public class AllDesktopPlatform : Platform
+{
+	public AllDesktopPlatform()
+		: base(UnrealTargetPlatform.AllDesktop)
+	{
+	}
+
+	public override bool CanBeCompiled()
+	{
+		return false;
+	}
+
+
+	public override UnrealTargetPlatform[] GetStagePlatforms()
+	{
+		return new UnrealTargetPlatform[] 
+		{
+			UnrealTargetPlatform.Win32,
+//			UnrealTargetPlatform.Win64,
+			UnrealTargetPlatform.Mac,
+			UnrealTargetPlatform.Linux,
+		};
+	}
+
+	public override void Package(ProjectParams Params, DeploymentContext SC, int WorkingCL)
+	{
+		SC.bIsCombiningMultiplePlatforms = true;
+		string SavedPlatformDir = SC.PlatformDir;
+		foreach (UnrealTargetPlatform DesktopPlatform in GetStagePlatforms())
+		{
+			Platform SubPlatform = Platform.Platforms[DesktopPlatform];
+			SC.PlatformDir = DesktopPlatform.ToString();
+			SubPlatform.Package(Params, SC, WorkingCL);
+		}
+		SC.PlatformDir = SavedPlatformDir;
+		SC.bIsCombiningMultiplePlatforms = false;
+	}
+
+	public override void GetFilesToArchive(ProjectParams Params, DeploymentContext SC)
+	{
+		SC.bIsCombiningMultiplePlatforms = true;
+		string SavedPlatformDir = SC.PlatformDir;
+		foreach (UnrealTargetPlatform DesktopPlatform in GetStagePlatforms())
+		{
+			Platform SubPlatform = Platform.Platforms[DesktopPlatform];
+			SC.PlatformDir = DesktopPlatform.ToString();
+			SubPlatform.GetFilesToArchive(Params, SC);
+		}
+		SC.PlatformDir = SavedPlatformDir;
+		SC.bIsCombiningMultiplePlatforms = false;
+	}
+
+	public override void GetConnectedDevices(ProjectParams Params, out List<string> Devices)
+	{
+		Devices = new List<string>();
+	}
+
+	public override ProcessResult RunClient(ERunOptions ClientRunFlags, string ClientApp, string ClientCmdLine, ProjectParams Params)
+	{
+		return null;
+	}
+
+	public override void GetFilesToDeployOrStage(ProjectParams Params, DeploymentContext SC)
+	{
+		SC.bIsCombiningMultiplePlatforms = true;
+		string SavedPlatformDir = SC.PlatformDir;
+		foreach (UnrealTargetPlatform DesktopPlatform in GetStagePlatforms())
+		{
+			Platform SubPlatform = Platform.Platforms[DesktopPlatform];
+			SC.PlatformDir = DesktopPlatform.ToString();
+			SubPlatform.GetFilesToDeployOrStage(Params, SC);
+		}
+		SC.PlatformDir = SavedPlatformDir;
+		SC.bIsCombiningMultiplePlatforms = false;
+	}
+
+	public override void ProcessArchivedProject(ProjectParams Params, DeploymentContext SC)
+	{
+		Console.WriteLine("***************************** PROCESSING ARCHIVED PROJECT ****************");
+
+		SC.bIsCombiningMultiplePlatforms = true;
+		string SavedPlatformDir = SC.PlatformDir;
+		foreach (UnrealTargetPlatform DesktopPlatform in GetStagePlatforms())
+		{
+			Platform SubPlatform = Platform.Platforms[DesktopPlatform];
+			SC.PlatformDir = DesktopPlatform.ToString();
+			SubPlatform.ProcessArchivedProject(Params, SC);
+		}
+		SC.PlatformDir = SavedPlatformDir;
+		SC.bIsCombiningMultiplePlatforms = false;
+
+		Console.WriteLine("***************************** DONE PROCESSING ARCHIVED PROJECT ****************");
+	}
+
+	/// <summary>
+	/// Gets cook platform name for this platform.
+	/// </summary>
+	/// <param name="CookFlavor">Additional parameter used to indicate special sub-target platform.</param>
+	/// <returns>Cook platform string.</returns>
+	public override string GetCookPlatform(bool bDedicatedServer, bool bIsClientOnly, string CookFlavor)
+	{
+		return "AllDesktop";
+	}
+
+	public override bool DeployPakInternalLowerCaseFilenames()
+	{
+		return false;
+	}
+
+	public override bool DeployLowerCaseFilenames(bool bUFSFile)
+	{
+		return false;
+	}
+
+	public override bool IsSupported { get { return true; } }
+
+	public override string Remap(string Dest)
+	{
+		return Dest;
+	}
+
+	public override PakType RequiresPak(ProjectParams Params)
+	{
+		return PakType.DontCare;
+	}
+    
+	public override List<string> GetDebugFileExtentions()
+	{
+		return new List<string> { };
+	}
+}
diff --git a/Engine/Source/Programs/AutomationTool/AllDesktop/Properties/AssemblyInfo.cs b/Engine/Source/Programs/AutomationTool/AllDesktop/Properties/AssemblyInfo.cs
new file mode 100644
--- /dev/null
+++ b/Engine/Source/Programs/AutomationTool/AllDesktop/Properties/AssemblyInfo.cs
@@ -0,0 +1,12 @@
+// Copyright 1998-2015 Epic Games, Inc. All Rights Reserved.
+
+using System.Reflection;
+using System.Runtime.CompilerServices;
+using System.Runtime.InteropServices;
+
+// General Information about an assembly is controlled through the following 
+// set of attributes. Change these attribute values to modify the information
+// associated with an assembly.
+[assembly: AssemblyTitle("Desktop.Automation")]
+[assembly: AssemblyDescription("")]
+[assembly: AssemblyConfiguration("")]
diff --git a/Engine/Source/Programs/AutomationTool/Desktop/Desktop.Automation.csproj b/Engine/Source/Programs/AutomationTool/Desktop/Desktop.Automation.csproj
deleted file mode 100644
diff --git a/Engine/Source/Programs/AutomationTool/Desktop/DesktopPlatform.Automation.cs b/Engine/Source/Programs/AutomationTool/Desktop/DesktopPlatform.Automation.cs
deleted file mode 100644
diff --git a/Engine/Source/Programs/AutomationTool/Desktop/Properties/AssemblyInfo.cs b/Engine/Source/Programs/AutomationTool/Desktop/Properties/AssemblyInfo.cs
deleted file mode 100644
diff --git a/Engine/Source/Programs/UnrealBuildTool/Configuration/UEBuildTarget.cs b/Engine/Source/Programs/UnrealBuildTool/Configuration/UEBuildTarget.cs
--- a/Engine/Source/Programs/UnrealBuildTool/Configuration/UEBuildTarget.cs
+++ b/Engine/Source/Programs/UnrealBuildTool/Configuration/UEBuildTarget.cs
@@ -27,7 +27,7 @@ namespace UnrealBuildTool
 		Android,
 		HTML5,
         Linux,
-		Desktop,
+		AllDesktop,
 	}
 
 	public enum UnrealPlatformGroup
@@ -46,7 +46,7 @@ namespace UnrealBuildTool
           */
         Device, 
         Simulator,
-		Desktop,
+		AllDesktop,
 	}
 
 	public enum UnrealTargetConfiguration
diff --git a/Engine/Source/Programs/UnrealBuildTool/Linux/UEBuildLinux.cs b/Engine/Source/Programs/UnrealBuildTool/Linux/UEBuildLinux.cs
--- a/Engine/Source/Programs/UnrealBuildTool/Linux/UEBuildLinux.cs
+++ b/Engine/Source/Programs/UnrealBuildTool/Linux/UEBuildLinux.cs
@@ -327,7 +327,7 @@ namespace UnrealBuildTool
                     InModule.AddDynamicallyLoadedModule("LinuxTargetPlatform");
                     InModule.AddDynamicallyLoadedModule("LinuxNoEditorTargetPlatform");
                     InModule.AddDynamicallyLoadedModule("LinuxServerTargetPlatform");
-					InModule.AddDynamicallyLoadedModule("DesktopTargetPlatform");
+					InModule.AddDynamicallyLoadedModule("AllDesktopTargetPlatform");
                 }
 
                 if (bBuildShaderFormats)
diff --git a/Engine/Source/Programs/UnrealBuildTool/Mac/UEBuildMac.cs b/Engine/Source/Programs/UnrealBuildTool/Mac/UEBuildMac.cs
--- a/Engine/Source/Programs/UnrealBuildTool/Mac/UEBuildMac.cs
+++ b/Engine/Source/Programs/UnrealBuildTool/Mac/UEBuildMac.cs
@@ -113,7 +113,7 @@ namespace UnrealBuildTool
 					InModule.AddDynamicallyLoadedModule("MacNoEditorTargetPlatform");
 					InModule.AddDynamicallyLoadedModule("MacClientTargetPlatform");
 					InModule.AddDynamicallyLoadedModule("MacServerTargetPlatform");
-					InModule.AddDynamicallyLoadedModule("DesktopTargetPlatform");
+					InModule.AddDynamicallyLoadedModule("AllDesktopTargetPlatform");
 				}
 
                 if (bBuildShaderFormats)
diff --git a/Engine/Source/Programs/UnrealBuildTool/Windows/UEBuildWindows.cs b/Engine/Source/Programs/UnrealBuildTool/Windows/UEBuildWindows.cs
--- a/Engine/Source/Programs/UnrealBuildTool/Windows/UEBuildWindows.cs
+++ b/Engine/Source/Programs/UnrealBuildTool/Windows/UEBuildWindows.cs
@@ -480,7 +480,7 @@ namespace UnrealBuildTool
                     InModule.AddDynamicallyLoadedModule("WindowsNoEditorTargetPlatform");
                     InModule.AddDynamicallyLoadedModule("WindowsServerTargetPlatform");
                     InModule.AddDynamicallyLoadedModule("WindowsClientTargetPlatform");
-					InModule.AddDynamicallyLoadedModule("DesktopTargetPlatform");
+					InModule.AddDynamicallyLoadedModule("AllDesktopTargetPlatform");
                 }
 
                 if (bBuildShaderFormats)
diff --git a/Engine/Source/Runtime/Engine/Engine.Build.cs b/Engine/Source/Runtime/Engine/Engine.Build.cs
--- a/Engine/Source/Runtime/Engine/Engine.Build.cs
+++ b/Engine/Source/Runtime/Engine/Engine.Build.cs
@@ -149,7 +149,7 @@ public class Engine : ModuleRules
 						"WindowsNoEditorTargetPlatform",
 						"WindowsServerTargetPlatform",
 						"WindowsClientTargetPlatform",
-						"DesktopTargetPlatform",
+						"AllDesktopTargetPlatform",
 					}
 				);
 			}
@@ -161,7 +161,7 @@ public class Engine : ModuleRules
 					    "MacNoEditorTargetPlatform",
 						"MacServerTargetPlatform",
 						"MacClientTargetPlatform",
-						"DesktopTargetPlatform",
+						"AllDesktopTargetPlatform",
 					}
 				);
 			}
@@ -172,7 +172,7 @@ public class Engine : ModuleRules
 						"LinuxTargetPlatform",
 						"LinuxNoEditorTargetPlatform",
 						"LinuxServerTargetPlatform",
-						"DesktopTargetPlatform",
+						"AllDesktopTargetPlatform",
 					}
 				);
 			}
